Collaborative Viewer template integrate in the wizard template

diff --git a/Templates/TP_Collaborative_V/Config/DefaultEngine.ini b/Templates/TP_Collaborative_V/Config/DefaultEngine.ini
new file mode 100644
--- /dev/null
+++ b/Templates/TP_Collaborative_V/Config/DefaultEngine.ini
@@ -0,0 +1,108 @@
+[URL]
+GameName=CollaborationDesignViewer
+
+[/Script/EngineSettings.GameMapsSettings]
+GameInstanceClass=/Game/CollaborativeViewer/Blueprints/GameInstance/BP_CollaborativeViewer_GameInstance.BP_CollaborativeViewer_GameInstance_C
+GlobalDefaultGameMode=/Game/CollaborativeViewer/Blueprints/GameMode/BP_CollaborativeViewer_GameMode.BP_CollaborativeViewer_GameMode_C
+EditorStartupMap=/Game/CollaborativeViewer/Maps/CollaborativeViewer_P.CollaborativeViewer_P
+GameDefaultMap=/Game/CollaborativeViewer/Maps/MainMenu.MainMenu
+
+[/Script/Engine.RendererSettings]
+r.DefaultFeature.AutoExposure.ExtendDefaultLuminanceRange=True
+r.CustomDepth=3
+
+[/Script/Engine.UserInterfaceSettings]
+UIScaleRule=ShortestSide
+UIScaleCurve=(EditorCurveData=(Keys=((Time=480.000000,Value=0.444000),(Time=720.000000,Value=0.666000),(Time=1080.000000,Value=1.000000),(Time=8640.000000,Value=8.000000)),PreInfinityExtrap=RCCE_Constant,PostInfinityExtrap=RCCE_Constant,DefaultValue=340282346638528859811704183484516925440.000000),ExternalCurve=None)
+bAllowHighDPIInGameMode=False
+ApplicationScale=1.000000
+
+[/Script/IOSRuntimeSettings.IOSRuntimeSettings]
+RemoteServerName=10.9.101.128
+RSyncUsername=compile
+MobileProvision=
+SigningCertificate=
+
+[/Script/Engine.Engine]
+NearClipPlane=1.000000
+
+[/Script/Engine.CollisionProfile]
++Profiles=(Name="NoCollision",CollisionEnabled=NoCollision,ObjectTypeName="WorldStatic",CustomResponses=((Channel="Visibility",Response=ECR_Ignore),(Channel="Camera",Response=ECR_Ignore)),HelpMessage="No collision",bCanModify=False)
++Profiles=(Name="BlockAll",CollisionEnabled=QueryAndPhysics,ObjectTypeName="WorldStatic",CustomResponses=,HelpMessage="WorldStatic object that blocks all actors by default. All new custom channels will use its own default response. ",bCanModify=False)
++Profiles=(Name="OverlapAll",CollisionEnabled=QueryOnly,ObjectTypeName="WorldStatic",CustomResponses=((Channel="WorldStatic",Response=ECR_Overlap),(Channel="Pawn",Response=ECR_Overlap),(Channel="Visibility",Response=ECR_Overlap),(Channel="WorldDynamic",Response=ECR_Overlap),(Channel="Camera",Response=ECR_Overlap),(Channel="PhysicsBody",Response=ECR_Overlap),(Channel="Vehicle",Response=ECR_Overlap),(Channel="Destructible",Response=ECR_Overlap)),HelpMessage="WorldStatic object that overlaps all actors by default. All new custom channels will use its own default response. ",bCanModify=False)
++Profiles=(Name="BlockAllDynamic",CollisionEnabled=QueryAndPhysics,ObjectTypeName="WorldDynamic",CustomResponses=,HelpMessage="WorldDynamic object that blocks all actors by default. All new custom channels will use its own default response. ",bCanModify=False)
++Profiles=(Name="OverlapAllDynamic",CollisionEnabled=QueryOnly,ObjectTypeName="WorldDynamic",CustomResponses=((Channel="WorldStatic",Response=ECR_Overlap),(Channel="Pawn",Response=ECR_Overlap),(Channel="Visibility",Response=ECR_Overlap),(Channel="WorldDynamic",Response=ECR_Overlap),(Channel="Camera",Response=ECR_Overlap),(Channel="PhysicsBody",Response=ECR_Overlap),(Channel="Vehicle",Response=ECR_Overlap),(Channel="Destructible",Response=ECR_Overlap)),HelpMessage="WorldDynamic object that overlaps all actors by default. All new custom channels will use its own default response. ",bCanModify=False)
++Profiles=(Name="IgnoreOnlyPawn",CollisionEnabled=QueryOnly,ObjectTypeName="WorldDynamic",CustomResponses=((Channel="Pawn",Response=ECR_Ignore),(Channel="Vehicle",Response=ECR_Ignore)),HelpMessage="WorldDynamic object that ignores Pawn and Vehicle. All other channels will be set to default.",bCanModify=False)
++Profiles=(Name="OverlapOnlyPawn",CollisionEnabled=QueryOnly,ObjectTypeName="WorldDynamic",CustomResponses=((Channel="Pawn",Response=ECR_Overlap),(Channel="Vehicle",Response=ECR_Overlap),(Channel="Camera",Response=ECR_Ignore)),HelpMessage="WorldDynamic object that overlaps Pawn, Camera, and Vehicle. All other channels will be set to default. ",bCanModify=False)
++Profiles=(Name="Pawn",CollisionEnabled=QueryAndPhysics,ObjectTypeName="Pawn",CustomResponses=((Channel="Visibility",Response=ECR_Ignore)),HelpMessage="Pawn object. Can be used for capsule of any playerable character or AI. ",bCanModify=False)
++Profiles=(Name="CharacterMesh",CollisionEnabled=QueryOnly,ObjectTypeName="Pawn",CustomResponses=((Channel="Pawn",Response=ECR_Ignore),(Channel="Vehicle",Response=ECR_Ignore),(Channel="Visibility",Response=ECR_Ignore)),HelpMessage="Pawn object that is used for Character Mesh. All other channels will be set to default.",bCanModify=False)
++Profiles=(Name="PhysicsActor",CollisionEnabled=QueryAndPhysics,ObjectTypeName="PhysicsBody",CustomResponses=,HelpMessage="Simulating actors",bCanModify=False)
++Profiles=(Name="Destructible",CollisionEnabled=QueryAndPhysics,ObjectTypeName="Destructible",CustomResponses=,HelpMessage="Destructible actors",bCanModify=False)
++Profiles=(Name="InvisibleWall",CollisionEnabled=QueryAndPhysics,ObjectTypeName="WorldStatic",CustomResponses=((Channel="Visibility",Response=ECR_Ignore)),HelpMessage="WorldStatic object that is invisible.",bCanModify=False)
++Profiles=(Name="InvisibleWallDynamic",CollisionEnabled=QueryAndPhysics,ObjectTypeName="WorldDynamic",CustomResponses=((Channel="Visibility",Response=ECR_Ignore)),HelpMessage="WorldDynamic object that is invisible.",bCanModify=False)
++Profiles=(Name="Trigger",CollisionEnabled=QueryOnly,ObjectTypeName="WorldDynamic",CustomResponses=((Channel="WorldStatic",Response=ECR_Overlap),(Channel="Pawn",Response=ECR_Overlap),(Channel="Visibility",Response=ECR_Ignore),(Channel="WorldDynamic",Response=ECR_Overlap),(Channel="Camera",Response=ECR_Overlap),(Channel="PhysicsBody",Response=ECR_Overlap),(Channel="Vehicle",Response=ECR_Overlap),(Channel="Destructible",Response=ECR_Overlap)),HelpMessage="WorldDynamic object that is used for trigger. All other channels will be set to default.",bCanModify=False)
++Profiles=(Name="Ragdoll",CollisionEnabled=QueryAndPhysics,ObjectTypeName="PhysicsBody",CustomResponses=((Channel="Pawn",Response=ECR_Ignore),(Channel="Visibility",Response=ECR_Ignore)),HelpMessage="Simulating Skeletal Mesh Component. All other channels will be set to default.",bCanModify=False)
++Profiles=(Name="Vehicle",CollisionEnabled=QueryAndPhysics,ObjectTypeName="Vehicle",CustomResponses=,HelpMessage="Vehicle object that blocks Vehicle, WorldStatic, and WorldDynamic. All other channels will be set to default.",bCanModify=False)
++Profiles=(Name="Spectator",CollisionEnabled=QueryOnly,ObjectTypeName="Pawn",CustomResponses=((Channel="WorldStatic"),(Channel="Pawn",Response=ECR_Ignore),(Channel="Visibility",Response=ECR_Ignore),(Channel="WorldDynamic",Response=ECR_Ignore),(Channel="Camera",Response=ECR_Ignore),(Channel="PhysicsBody",Response=ECR_Ignore),(Channel="Vehicle",Response=ECR_Ignore),(Channel="Destructible",Response=ECR_Ignore)),HelpMessage="Pawn object that ignores all other actors except WorldStatic.",bCanModify=False)
++Profiles=(Name="UI",CollisionEnabled=QueryOnly,ObjectTypeName="WorldDynamic",CustomResponses=((Channel="WorldStatic",Response=ECR_Overlap),(Channel="Pawn",Response=ECR_Overlap),(Channel="Visibility"),(Channel="WorldDynamic",Response=ECR_Overlap),(Channel="Camera",Response=ECR_Overlap),(Channel="PhysicsBody",Response=ECR_Overlap),(Channel="Vehicle",Response=ECR_Overlap),(Channel="Destructible",Response=ECR_Overlap)),HelpMessage="WorldStatic object that overlaps all actors by default. All new custom channels will use its own default response. ",bCanModify=False)
++DefaultChannelResponses=(Channel=ECC_GameTraceChannel1,Name="Annotation",DefaultResponse=ECR_Block,bTraceType=False,bStaticObject=False)
++DefaultChannelResponses=(Channel=ECC_GameTraceChannel2,Name="AnnotationStroke",DefaultResponse=ECR_Block,bTraceType=False,bStaticObject=False)
++DefaultChannelResponses=(Channel=ECC_GameTraceChannel3,Name="SnapShot",DefaultResponse=ECR_Block,bTraceType=False,bStaticObject=False)
++DefaultChannelResponses=(Channel=ECC_GameTraceChannel4,Name="DimenssionTool",DefaultResponse=ECR_Block,bTraceType=False,bStaticObject=False)
++ProfileRedirects=(OldName="BlockingVolume",NewName="InvisibleWall")
++ProfileRedirects=(OldName="InterpActor",NewName="IgnoreOnlyPawn")
++ProfileRedirects=(OldName="StaticMeshComponent",NewName="BlockAllDynamic")
++ProfileRedirects=(OldName="SkeletalMeshActor",NewName="PhysicsActor")
++ProfileRedirects=(OldName="InvisibleActor",NewName="InvisibleWallDynamic")
++CollisionChannelRedirects=(OldName="Static",NewName="WorldStatic")
++CollisionChannelRedirects=(OldName="Dynamic",NewName="WorldDynamic")
++CollisionChannelRedirects=(OldName="VehicleMovement",NewName="Vehicle")
++CollisionChannelRedirects=(OldName="PawnMovement",NewName="Pawn")
+
+[/Script/Engine.PhysicsSettings]
+DefaultGravityZ=-980.000000
+DefaultTerminalVelocity=4000.000000
+DefaultFluidFriction=0.300000
+SimulateScratchMemorySize=262144
+RagdollAggregateThreshold=4
+TriangleMeshTriangleMinAreaThreshold=5.000000
+bEnableShapeSharing=False
+bEnablePCM=True
+bEnableStabilization=False
+bWarnMissingLocks=True
+bEnable2DPhysics=False
+PhysicErrorCorrection=(PingExtrapolation=0.100000,PingLimit=100.000000,ErrorPerLinearDifference=1.000000,ErrorPerAngularDifference=1.000000,MaxRestoredStateError=1.000000,MaxLinearHardSnapDistance=400.000000,PositionLerp=0.000000,AngleLerp=0.400000,LinearVelocityCoefficient=100.000000,AngularVelocityCoefficient=10.000000,ErrorAccumulationSeconds=0.500000,ErrorAccumulationDistanceSq=15.000000,ErrorAccumulationSimilarity=100.000000)
+LockedAxis=Invalid
+DefaultDegreesOfFreedom=Full3D
+BounceThresholdVelocity=200.000000
+FrictionCombineMode=Average
+RestitutionCombineMode=Average
+MaxAngularVelocity=3600.000000
+MaxDepenetrationVelocity=0.000000
+ContactOffsetMultiplier=0.020000
+MinContactOffset=2.000000
+MaxContactOffset=8.000000
+bSimulateSkeletalMeshOnDedicatedServer=True
+DefaultShapeComplexity=CTF_UseSimpleAndComplex
+bDefaultHasComplexCollision=True
+bSuppressFaceRemapTable=False
+bSupportUVFromHitResults=False
+bDisableActiveActors=False
+bDisableKinematicStaticPairs=False
+bDisableKinematicKinematicPairs=False
+bDisableCCD=False
+bEnableEnhancedDeterminism=False
+AnimPhysicsMinDeltaTime=0.000000
+bSimulateAnimPhysicsAfterReset=False
+MaxPhysicsDeltaTime=0.033333
+bSubstepping=False
+bSubsteppingAsync=False
+MaxSubstepDeltaTime=0.016667
+MaxSubsteps=6
+SyncSceneSmoothingFactor=0.000000
+InitialAverageFrameRate=0.016667
+PhysXTreeRebuildRate=10
+DefaultBroadphaseSettings=(bUseMBPOnClient=False,bUseMBPOnServer=False,bUseMBPOuterBounds=False,MBPBounds=(Min=(X=0.000000,Y=0.000000,Z=0.000000),Max=(X=0.000000,Y=0.000000,Z=0.000000),IsValid=0),MBPOuterBounds=(Min=(X=0.000000,Y=0.000000,Z=0.000000),Max=(X=0.000000,Y=0.000000,Z=0.000000),IsValid=0),MBPNumSubdivs=2)
+ChaosSettings=(DefaultThreadingModel=TaskGraph,DedicatedThreadTickMode=VariableCappedWithTarget,DedicatedThreadBufferMode=Double)
+
+
diff --git a/Templates/TP_Collaborative_V/Config/DefaultGame.ini b/Templates/TP_Collaborative_V/Config/DefaultGame.ini
new file mode 100644
--- /dev/null
+++ b/Templates/TP_Collaborative_V/Config/DefaultGame.ini
@@ -0,0 +1,85 @@
+
+[/Script/UnrealEd.ProjectPackagingSettings]
+Build=IfProjectHasCode
+BuildConfiguration=PPBC_Development
+StagingDirectory=(Path="")
+FullRebuild=False
+ForDistribution=False
+IncludeDebugFiles=False
+BlueprintNativizationMethod=Disabled
+bIncludeNativizedAssetsInProjectGeneration=False
+bExcludeMonolithicEngineHeadersInNativizedCode=False
+UsePakFile=True
+bGenerateChunks=False
+bGenerateNoChunks=False
+bChunkHardReferencesOnly=False
+bForceOneChunkPerFile=False
+MaxChunkSize=0
+bBuildHttpChunkInstallData=False
+HttpChunkInstallDataDirectory=(Path="")
+PakFileCompressionFormats=
+PakFileAdditionalCompressionOptions=
+HttpChunkInstallDataVersion=
+IncludePrerequisites=True
+IncludeAppLocalPrerequisites=False
+bShareMaterialShaderCode=False
+bSharedMaterialNativeLibraries=False
+ApplocalPrerequisitesDirectory=(Path="")
+IncludeCrashReporter=False
+InternationalizationPreset=English
+-CulturesToStage=en
++CulturesToStage=en
+bCookAll=False
+bCookMapsOnly=False
+bCompressed=False
+bEncryptIniFiles=False
+bEncryptPakIndex=False
+GenerateEarlyDownloaderPakFile=False
+bSkipEditorContent=False
+bSkipMovies=False
+-EarlyDownloaderPakFileFiles=...\Content\Internationalization\...\*.icu
+-EarlyDownloaderPakFileFiles=...\Content\Internationalization\...\*.brk
+-EarlyDownloaderPakFileFiles=...\Content\Internationalization\...\*.res
+-EarlyDownloaderPakFileFiles=...\Content\Internationalization\...\*.nrm
+-EarlyDownloaderPakFileFiles=...\Content\Internationalization\...\*.cfu
+-EarlyDownloaderPakFileFiles=...\Content\Localization\...\*.*
+-EarlyDownloaderPakFileFiles=...\Content\Localization\*.*
+-EarlyDownloaderPakFileFiles=...\Content\Certificates\...\*.*
+-EarlyDownloaderPakFileFiles=...\Content\Certificates\*.*
+-EarlyDownloaderPakFileFiles=-...\Content\Localization\Game\...\*.*
+-EarlyDownloaderPakFileFiles=-...\Content\Localization\Game\*.*
+-EarlyDownloaderPakFileFiles=...\Config\...\*.ini
+-EarlyDownloaderPakFileFiles=...\Config\*.ini
+-EarlyDownloaderPakFileFiles=...\Engine\GlobalShaderCache*.bin
+-EarlyDownloaderPakFileFiles=...\Content\ShaderArchive-Global*.ushaderbytecode
+-EarlyDownloaderPakFileFiles=...\Content\Slate\*.*
+-EarlyDownloaderPakFileFiles=...\Content\Slate\...\*.*
+-EarlyDownloaderPakFileFiles=...\*.upluginmanifest
+-EarlyDownloaderPakFileFiles=...\*.uproject
+-EarlyDownloaderPakFileFiles=...\global_sf*.metalmap
++EarlyDownloaderPakFileFiles=...\Content\Internationalization\...\*.icu
++EarlyDownloaderPakFileFiles=...\Content\Internationalization\...\*.brk
++EarlyDownloaderPakFileFiles=...\Content\Internationalization\...\*.res
++EarlyDownloaderPakFileFiles=...\Content\Internationalization\...\*.nrm
++EarlyDownloaderPakFileFiles=...\Content\Internationalization\...\*.cfu
++EarlyDownloaderPakFileFiles=...\Content\Localization\...\*.*
++EarlyDownloaderPakFileFiles=...\Content\Localization\*.*
++EarlyDownloaderPakFileFiles=...\Content\Certificates\...\*.*
++EarlyDownloaderPakFileFiles=...\Content\Certificates\*.*
++EarlyDownloaderPakFileFiles=-...\Content\Localization\Game\...\*.*
++EarlyDownloaderPakFileFiles=-...\Content\Localization\Game\*.*
++EarlyDownloaderPakFileFiles=...\Config\...\*.ini
++EarlyDownloaderPakFileFiles=...\Config\*.ini
++EarlyDownloaderPakFileFiles=...\Engine\GlobalShaderCache*.bin
++EarlyDownloaderPakFileFiles=...\Content\ShaderArchive-Global*.ushaderbytecode
++EarlyDownloaderPakFileFiles=...\Content\Slate\*.*
++EarlyDownloaderPakFileFiles=...\Content\Slate\...\*.*
++EarlyDownloaderPakFileFiles=...\*.upluginmanifest
++EarlyDownloaderPakFileFiles=...\*.uproject
++EarlyDownloaderPakFileFiles=...\global_sf*.metalmap
++MapsToCook=(FilePath="/Game/CollaborativeViewer/Maps/CollaborativeViewer_P")
++MapsToCook=(FilePath="/Game/CollaborativeViewer/Maps/MainMenu")
+bNativizeBlueprintAssets=False
+bNativizeOnlySelectedBlueprints=False
+
+
diff --git a/Templates/TP_Collaborative_V/Config/DefaultInput.ini b/Templates/TP_Collaborative_V/Config/DefaultInput.ini
new file mode 100644
--- /dev/null
+++ b/Templates/TP_Collaborative_V/Config/DefaultInput.ini
@@ -0,0 +1,116 @@
+
+[/Script/Engine.InputSettings]
++AxisConfig=(AxisKeyName="OculusTouch_Right_Thumbstick",AxisProperties=(DeadZone=0.000000,Sensitivity=1.000000,Exponent=1.000000,bInvert=False))
++AxisConfig=(AxisKeyName="OculusTouch_Right_FaceButton1",AxisProperties=(DeadZone=0.000000,Sensitivity=1.000000,Exponent=1.000000,bInvert=False))
++AxisConfig=(AxisKeyName="OculusTouch_Right_Trigger",AxisProperties=(DeadZone=0.000000,Sensitivity=1.000000,Exponent=1.000000,bInvert=False))
++AxisConfig=(AxisKeyName="OculusTouch_Right_FaceButton2",AxisProperties=(DeadZone=0.000000,Sensitivity=1.000000,Exponent=1.000000,bInvert=False))
++AxisConfig=(AxisKeyName="OculusTouch_Right_IndexPointing",AxisProperties=(DeadZone=0.000000,Sensitivity=1.000000,Exponent=1.000000,bInvert=False))
++AxisConfig=(AxisKeyName="OculusTouch_Right_ThumbUp",AxisProperties=(DeadZone=0.000000,Sensitivity=1.000000,Exponent=1.000000,bInvert=False))
++AxisConfig=(AxisKeyName="OculusTouchpad_Touchpad_X",AxisProperties=(DeadZone=0.000000,Sensitivity=1.000000,Exponent=1.000000,bInvert=False))
++AxisConfig=(AxisKeyName="OculusTouchpad_Touchpad_Y",AxisProperties=(DeadZone=0.000000,Sensitivity=1.000000,Exponent=1.000000,bInvert=False))
++AxisConfig=(AxisKeyName="SteamVR_Knuckles_Left_HandGrip",AxisProperties=(DeadZone=0.000000,Sensitivity=1.000000,Exponent=1.000000,bInvert=False))
++AxisConfig=(AxisKeyName="SteamVR_Knuckles_Left_IndexGrip",AxisProperties=(DeadZone=0.000000,Sensitivity=1.000000,Exponent=1.000000,bInvert=False))
++AxisConfig=(AxisKeyName="SteamVR_Knuckles_Left_MiddleGrip",AxisProperties=(DeadZone=0.000000,Sensitivity=1.000000,Exponent=1.000000,bInvert=False))
++AxisConfig=(AxisKeyName="SteamVR_Knuckles_Left_RingGrip",AxisProperties=(DeadZone=0.000000,Sensitivity=1.000000,Exponent=1.000000,bInvert=False))
++AxisConfig=(AxisKeyName="SteamVR_Knuckles_Left_PinkyGrip",AxisProperties=(DeadZone=0.000000,Sensitivity=1.000000,Exponent=1.000000,bInvert=False))
++AxisConfig=(AxisKeyName="SteamVR_Knuckles_Right_HandGrip",AxisProperties=(DeadZone=0.000000,Sensitivity=1.000000,Exponent=1.000000,bInvert=False))
++AxisConfig=(AxisKeyName="SteamVR_Knuckles_Right_IndexGrip",AxisProperties=(DeadZone=0.000000,Sensitivity=1.000000,Exponent=1.000000,bInvert=False))
++AxisConfig=(AxisKeyName="SteamVR_Knuckles_Right_MiddleGrip",AxisProperties=(DeadZone=0.000000,Sensitivity=1.000000,Exponent=1.000000,bInvert=False))
++AxisConfig=(AxisKeyName="SteamVR_Knuckles_Right_RingGrip",AxisProperties=(DeadZone=0.000000,Sensitivity=1.000000,Exponent=1.000000,bInvert=False))
++AxisConfig=(AxisKeyName="SteamVR_Knuckles_Right_PinkyGrip",AxisProperties=(DeadZone=0.000000,Sensitivity=1.000000,Exponent=1.000000,bInvert=False))
++AxisConfig=(AxisKeyName="Gamepad_LeftX",AxisProperties=(DeadZone=0.250000,Sensitivity=1.000000,Exponent=1.000000,bInvert=False))
++AxisConfig=(AxisKeyName="Gamepad_LeftY",AxisProperties=(DeadZone=0.250000,Sensitivity=1.000000,Exponent=1.000000,bInvert=False))
++AxisConfig=(AxisKeyName="Gamepad_RightX",AxisProperties=(DeadZone=0.250000,Sensitivity=1.000000,Exponent=1.000000,bInvert=False))
++AxisConfig=(AxisKeyName="Gamepad_RightY",AxisProperties=(DeadZone=0.250000,Sensitivity=1.000000,Exponent=1.000000,bInvert=False))
++AxisConfig=(AxisKeyName="MouseX",AxisProperties=(DeadZone=0.000000,Sensitivity=0.070000,Exponent=1.000000,bInvert=False))
++AxisConfig=(AxisKeyName="MouseY",AxisProperties=(DeadZone=0.000000,Sensitivity=0.070000,Exponent=1.000000,bInvert=False))
++AxisConfig=(AxisKeyName="MouseWheelAxis",AxisProperties=(DeadZone=0.000000,Sensitivity=1.000000,Exponent=1.000000,bInvert=False))
++AxisConfig=(AxisKeyName="Gamepad_LeftTriggerAxis",AxisProperties=(DeadZone=0.000000,Sensitivity=1.000000,Exponent=1.000000,bInvert=False))
++AxisConfig=(AxisKeyName="Gamepad_RightTriggerAxis",AxisProperties=(DeadZone=0.000000,Sensitivity=1.000000,Exponent=1.000000,bInvert=False))
++AxisConfig=(AxisKeyName="MotionController_Left_Thumbstick_X",AxisProperties=(DeadZone=0.000000,Sensitivity=1.000000,Exponent=1.000000,bInvert=False))
++AxisConfig=(AxisKeyName="MotionController_Left_Thumbstick_Y",AxisProperties=(DeadZone=0.000000,Sensitivity=1.000000,Exponent=1.000000,bInvert=False))
++AxisConfig=(AxisKeyName="MotionController_Left_TriggerAxis",AxisProperties=(DeadZone=0.000000,Sensitivity=1.000000,Exponent=1.000000,bInvert=False))
++AxisConfig=(AxisKeyName="MotionController_Left_Grip1Axis",AxisProperties=(DeadZone=0.000000,Sensitivity=1.000000,Exponent=1.000000,bInvert=False))
++AxisConfig=(AxisKeyName="MotionController_Left_Grip2Axis",AxisProperties=(DeadZone=0.000000,Sensitivity=1.000000,Exponent=1.000000,bInvert=False))
++AxisConfig=(AxisKeyName="MotionController_Right_Thumbstick_X",AxisProperties=(DeadZone=0.000000,Sensitivity=1.000000,Exponent=1.000000,bInvert=False))
++AxisConfig=(AxisKeyName="MotionController_Right_Thumbstick_Y",AxisProperties=(DeadZone=0.000000,Sensitivity=1.000000,Exponent=1.000000,bInvert=False))
++AxisConfig=(AxisKeyName="MotionController_Right_TriggerAxis",AxisProperties=(DeadZone=0.000000,Sensitivity=1.000000,Exponent=1.000000,bInvert=False))
++AxisConfig=(AxisKeyName="MotionController_Right_Grip1Axis",AxisProperties=(DeadZone=0.000000,Sensitivity=1.000000,Exponent=1.000000,bInvert=False))
++AxisConfig=(AxisKeyName="MotionController_Right_Grip2Axis",AxisProperties=(DeadZone=0.000000,Sensitivity=1.000000,Exponent=1.000000,bInvert=False))
++AxisConfig=(AxisKeyName="Gamepad_Special_Left_X",AxisProperties=(DeadZone=0.000000,Sensitivity=1.000000,Exponent=1.000000,bInvert=False))
++AxisConfig=(AxisKeyName="Gamepad_Special_Left_Y",AxisProperties=(DeadZone=0.000000,Sensitivity=1.000000,Exponent=1.000000,bInvert=False))
++AxisConfig=(AxisKeyName="MotionController_Left_Thumbstick_Z",AxisProperties=(DeadZone=0.000000,Sensitivity=1.000000,Exponent=1.000000,bInvert=False))
++AxisConfig=(AxisKeyName="MotionController_Right_Thumbstick_Z",AxisProperties=(DeadZone=0.000000,Sensitivity=1.000000,Exponent=1.000000,bInvert=False))
++AxisConfig=(AxisKeyName="OculusTouch_Left_Thumbstick",AxisProperties=(DeadZone=0.000000,Sensitivity=1.000000,Exponent=1.000000,bInvert=False))
++AxisConfig=(AxisKeyName="OculusTouch_Left_FaceButton1",AxisProperties=(DeadZone=0.000000,Sensitivity=1.000000,Exponent=1.000000,bInvert=False))
++AxisConfig=(AxisKeyName="OculusTouch_Left_Trigger",AxisProperties=(DeadZone=0.000000,Sensitivity=1.000000,Exponent=1.000000,bInvert=False))
++AxisConfig=(AxisKeyName="OculusTouch_Left_FaceButton2",AxisProperties=(DeadZone=0.000000,Sensitivity=1.000000,Exponent=1.000000,bInvert=False))
++AxisConfig=(AxisKeyName="OculusTouch_Left_IndexPointing",AxisProperties=(DeadZone=0.000000,Sensitivity=1.000000,Exponent=1.000000,bInvert=False))
++AxisConfig=(AxisKeyName="OculusTouch_Left_ThumbUp",AxisProperties=(DeadZone=0.000000,Sensitivity=1.000000,Exponent=1.000000,bInvert=False))
+bAltEnterTogglesFullscreen=True
+bF11TogglesFullscreen=True
+bUseMouseForTouch=False
+bEnableMouseSmoothing=True
+bEnableFOVScaling=True
+bCaptureMouseOnLaunch=True
+bDefaultViewportMouseLock=False
+bAlwaysShowTouchInterface=False
+bShowConsoleOnFourFingerTap=True
+bEnableGestureRecognizer=True
+bUseAutocorrect=False
+DefaultViewportMouseCaptureMode=CapturePermanently_IncludingInitialMouseDown
+DefaultViewportMouseLockMode=LockOnCapture
+FOVScale=0.011110
+DoubleClickTime=0.200000
++ActionMappings=(ActionName="MotionControllerL",bShift=False,bCtrl=False,bAlt=False,bCmd=False,Key=MotionController_Left_Trigger)
++ActionMappings=(ActionName="MotionControllerR",bShift=False,bCtrl=False,bAlt=False,bCmd=False,Key=MotionController_Right_Trigger)
++ActionMappings=(ActionName="PushAndPull",bShift=False,bCtrl=False,bAlt=False,bCmd=False,Key=MouseWheelAxis)
++ActionMappings=(ActionName="Quit",bShift=False,bCtrl=False,bAlt=False,bCmd=False,Key=Escape)
++ActionMappings=(ActionName="FlyMode",bShift=False,bCtrl=True,bAlt=False,bCmd=False,Key=One)
++ActionMappings=(ActionName="WalkMode",bShift=False,bCtrl=True,bAlt=False,bCmd=False,Key=Two)
++ActionMappings=(ActionName="OrbitMode",bShift=False,bCtrl=True,bAlt=False,bCmd=False,Key=Three)
++ActionMappings=(ActionName="VRMode",bShift=False,bCtrl=True,bAlt=False,bCmd=False,Key=Four)
++ActionMappings=(ActionName="ResetLocationL",bShift=False,bCtrl=False,bAlt=False,bCmd=False,Key=MotionController_Left_Grip1)
++ActionMappings=(ActionName="Bookmark_1",bShift=False,bCtrl=False,bAlt=False,bCmd=False,Key=One)
++ActionMappings=(ActionName="Bookmark_2",bShift=False,bCtrl=False,bAlt=False,bCmd=False,Key=Two)
++ActionMappings=(ActionName="Bookmark_3",bShift=False,bCtrl=False,bAlt=False,bCmd=False,Key=Three)
++ActionMappings=(ActionName="Bookmark_4",bShift=False,bCtrl=False,bAlt=False,bCmd=False,Key=Four)
++ActionMappings=(ActionName="Bookmark_5",bShift=False,bCtrl=False,bAlt=False,bCmd=False,Key=Five)
++ActionMappings=(ActionName="Bookmark_6",bShift=False,bCtrl=False,bAlt=False,bCmd=False,Key=Six)
++ActionMappings=(ActionName="Bookmark_7",bShift=False,bCtrl=False,bAlt=False,bCmd=False,Key=Seven)
++ActionMappings=(ActionName="Bookmark_8",bShift=False,bCtrl=False,bAlt=False,bCmd=False,Key=Eight)
++ActionMappings=(ActionName="Bookmark_9",bShift=False,bCtrl=False,bAlt=False,bCmd=False,Key=Nine)
++ActionMappings=(ActionName="OpenContextualMenu",bShift=False,bCtrl=False,bAlt=False,bCmd=False,Key=SpaceBar)
++ActionMappings=(ActionName="OrbitAroundObject",bShift=False,bCtrl=False,bAlt=False,bCmd=False,Key=RightMouseButton)
++ActionMappings=(ActionName="OrbitZoomOut",bShift=False,bCtrl=False,bAlt=False,bCmd=False,Key=MouseScrollDown)
++ActionMappings=(ActionName="OrbitZoomIn",bShift=False,bCtrl=False,bAlt=False,bCmd=False,Key=MouseScrollUp)
++ActionMappings=(ActionName="ShootLaser",bShift=False,bCtrl=False,bAlt=False,bCmd=False,Key=LeftMouseButton)
++ActionMappings=(ActionName="RotateView",bShift=False,bCtrl=False,bAlt=False,bCmd=False,Key=RightMouseButton)
++ActionMappings=(ActionName="SetPointOfinterest",bShift=False,bCtrl=False,bAlt=False,bCmd=False,Key=MiddleMouseButton)
++ActionMappings=(ActionName="LockToAxis",bShift=False,bCtrl=False,bAlt=False,bCmd=False,Key=LeftShift)
++AxisMappings=(AxisName="WalkForward",Scale=1.000000,Key=W)
++AxisMappings=(AxisName="WalkRight",Scale=1.000000,Key=D)
++AxisMappings=(AxisName="WalkForward",Scale=-1.000000,Key=S)
++AxisMappings=(AxisName="WalkRight",Scale=-1.000000,Key=A)
++AxisMappings=(AxisName="LookUp",Scale=1.000000,Key=MouseY)
++AxisMappings=(AxisName="LookRight",Scale=1.000000,Key=MouseX)
++AxisMappings=(AxisName="FlyUp",Scale=1.000000,Key=E)
++AxisMappings=(AxisName="FlyUp",Scale=-1.000000,Key=Q)
++AxisMappings=(AxisName="MotionControllerThumbLeft_Y",Scale=-1.000000,Key=MotionController_Left_Thumbstick_Y)
++AxisMappings=(AxisName="MotionControllerThumbRight_Y",Scale=-1.000000,Key=MotionController_Right_Thumbstick_Y)
++AxisMappings=(AxisName="MotionControllerThumbLeft_X",Scale=1.000000,Key=MotionController_Left_Thumbstick_X)
++AxisMappings=(AxisName="MotionControllerThumbRight_X",Scale=1.000000,Key=MotionController_Right_Thumbstick_X)
++AxisMappings=(AxisName="WalkForward",Scale=1.000000,Key=Gamepad_LeftY)
++AxisMappings=(AxisName="WalkRight",Scale=1.000000,Key=Gamepad_LeftX)
++AxisMappings=(AxisName="LookUp",Scale=-1.000000,Key=Gamepad_RightY)
++AxisMappings=(AxisName="LookRight",Scale=1.000000,Key=Gamepad_RightX)
++AxisMappings=(AxisName="pinch",Scale=1.000000,Key=Gesture_Pinch)
++AxisMappings=(AxisName="WalkForward",Scale=1.000000,Key=Up)
++AxisMappings=(AxisName="WalkForward",Scale=-1.000000,Key=Down)
++AxisMappings=(AxisName="WalkRight",Scale=1.000000,Key=Right)
++AxisMappings=(AxisName="WalkRight",Scale=-1.000000,Key=Left)
+DefaultTouchInterface=/Game/CollaborativeViewer/Blueprints/Pawn/Touch/TouchInterface.TouchInterface
+ConsoleKey=None
++ConsoleKeys=Tilde
+
+
diff --git a/Templates/TP_Collaborative_V/Config/TemplateDefs.ini b/Templates/TP_Collaborative_V/Config/TemplateDefs.ini
new file mode 100644
--- /dev/null
+++ b/Templates/TP_Collaborative_V/Config/TemplateDefs.ini
@@ -0,0 +1,96 @@
+﻿; This file specifies how a template will be treated when creating a project from it.
+; It includes rules for ignoring files, replacing files and changing the contents of files.
+; This file is annotated to help with creating your own templates.
+
+[/Script/GameProjectGeneration.TemplateProjectDefs]
+
+; The categories that this template should appear in. The values in this list refer to the Key value given in TemplateCategories.ini.
+Categories=AEC
+
+; The localized display names for this template. Each entry specifies a Language code and Unicode text to display in the New Project dialog.
+LocalizedDisplayNames=(Language="en", Text="Collab Viewer")
+
+; The localized descriptions for this template, which is displayed when the template is selected.
+LocalizedDescriptions=(Language="en", Text="This project template provides navigation and interaction for various platforms (Desktop, VR). Use as a starting point for collaborative Industrial or Architectural projects.")
+; The sort key to use for the template to bypass the default alphabetical sorting.
+; SortKey="_1"
+
+; A list of folders to avoid copying when creating a project from this template. 
+; These should contain files only of concern to the template itself and not desired in the user's project.
+FoldersToIgnore=Media
+
+; The files to ignore copying, relative to the template root folder. 
+; %TEMPLATENAME% will be replaced by the template's folder name when resolving these. 
+; Filename comparisons are case-insensitive. 
+; For example, if this template is "TP_Blank" and the user creates "MyProject", then "TP_Blank.uproject" will be copied to "MyProject.uproject".
+FilesToIgnore="%TEMPLATENAME%.uproject"
+FilesToIgnore="%TEMPLATENAME%.png"
+FilesToIgnore="Config/TemplateDefs.ini"
+FilesToIgnore="Manifest.json"
+FilesToIgnore="contents.txt"
+
+; NOTE: For the purposes of the below examples, this template is considered to be named "TP_Blank", and the user has created a project called "MyProject".
+
+; Folders that will be renamed when creating a project from this template.
+; %PROJECTNAME% will be replaced by the project name specified by the user during creation when resolving these rules.
+; For example,  then "Source/TP_Blank" will be copied to "Source/MyProject".
+FolderRenames=(From="Source/%TEMPLATENAME%", To="Source/%PROJECTNAME%")
+FolderRenames=(From="Source/%TEMPLATENAME%Editor", To="Source/%PROJECTNAME%Editor")
+
+; Rules for renaming files when creating a project from this template. 
+; %TEMPLATENAME_UPPERCASE% and %PROJECTNAME_UPPERCASE% are the template and project names in upper case, eg. for "TP_Blank", its upper case representation is "TP_BLANK".
+; %TEMPLATENAME_LOWERCASE% and %PROJECTNAME_LOWERCASE% are the template and project names in lower case, eg. for "TP_Blank", its lower case representation is "tp_blank".
+; A rule does not have to include any of these modifiers. 
+; Extensions is a list of files to consider for these replacements. Eg. ("cpp", "h") will try to apply a rule only to files like "SomeFile.cpp" and "SomeFile.h".
+; bCaseSensitive specifies whether this replacement should be case-sensitive.
+
+; Here are some default rules for renaming files:
+
+; This rule states that a file called "TP_BLANK.cpp" will be renamed to "MYPROJECT.cpp".
+FilenameReplacements=(Extensions=("cpp","h","ini","cs"), From="%TEMPLATENAME_UPPERCASE%", To="%PROJECTNAME_UPPERCASE%", bCaseSensitive=true)
+
+; This rule states that a file called "tp_blank.ini" will be renamed to "myproject.ini"
+FilenameReplacements=(Extensions=("cpp","h","ini","cs"), From="%TEMPLATENAME_LOWERCASE%", To="%PROJECTNAME_LOWERCASE%", bCaseSensitive=true)
+
+; And this rule catches any irregularly-cased files, so that "Tp_bLANK.ini" will be renamed to "MyProject.ini".
+FilenameReplacements=(Extensions=("cpp","h","ini","cs"), From="%TEMPLATENAME%",To="%PROJECTNAME%",bCaseSensitive=false)
+
+; Rules for renaming strings in the contents of files. 
+; The same rules as above apply, but any usages of strings that match these rules inside the file will be changed. This includes all substrings!
+; For example, "TP_Blank" will be replaced by "MyProject", but "TP_Blank_Name" will also be replaced by "MyProject_Name".
+ReplacementsInFiles=(Extensions=("cpp","h","ini","cs"), From="%TEMPLATENAME_UPPERCASE%", To="%PROJECTNAME_UPPERCASE%", bCaseSensitive=true)
+ReplacementsInFiles=(Extensions=("cpp","h","ini","cs"), From="%TEMPLATENAME_LOWERCASE%", To="%PROJECTNAME_LOWERCASE%", bCaseSensitive=true)
+ReplacementsInFiles=(Extensions=("cpp","h","ini","cs"), From="%TEMPLATENAME%", To="%PROJECTNAME%", bCaseSensitive=false)
+
+; The following settings are not enabled for this project, but are included for completeness:
+
+; Whether or not this template can be used to create a project. Default: true
+; If this is set to false, then the template becomes a Feature Pack that can only be referenced from other templates, or added to a project from the Content Browser.
+;bAllowProjectCreation=false
+
+; Which settings on the project settings page should be hidden for this template. Default: none
+; Allowed values are:
+; All - Hides the settings page entirely
+; Languages - Choose between C++ and Blueprint versions.
+; HardwareTarget - Choose between Mobile and Desktop hardware targets.
+; GraphicsPreset - Choose between Maximum and Scalable graphics.
+; StarterContent - Choose whether to add starter content to the project.
+; VR - Choose whether to enable Virtual Reality support for the project.
+;HiddenSettings=All
+HiddenSettings=Languages
+HiddenSettings=HardwareTarget
+HiddenSettings=GraphicsPreset 
+;HiddenSettings=StarterContent 
+;HiddenSettings=XR
+
+
+; The class types that this project uses.
+; This is only for display in the template description, and does not affect project creation.
+;ClassTypes=GameMode, Character, Actor
+
+; The asset types that this project uses.
+; This is only for display in the template description, and does not affect project creation.
+;AssetTypes=Skeleton, Static Mesh, Texture
+
+; Shared feature packs. The files in these packs listed in these structures marked as 'additionalfiles' will be copied on project generation.
+;SharedContentPacks=(MountName="2DSideScroller",DetailLevels=("Standard"))
diff --git a/Templates/TP_Collaborative_V/TP_Collaborative_V.uproject b/Templates/TP_Collaborative_V/TP_Collaborative_V.uproject
new file mode 100644
--- /dev/null
+++ b/Templates/TP_Collaborative_V/TP_Collaborative_V.uproject
@@ -0,0 +1,6 @@
+{
+	"FileVersion": 3,
+	"EngineAssociation": "",
+	"Category": "",
+	"Description": ""
+}
\ No newline at end of file
diff --git a/Templates/TP_Collaborative_V/contents.txt b/Templates/TP_Collaborative_V/contents.txt
new file mode 100644
--- /dev/null
+++ b/Templates/TP_Collaborative_V/contents.txt
@@ -0,0 +1,4 @@
+../../../Templates/TP_Collaborative_V/manifest.json
+../../../Templates/TP_Collaborative_V/Config/config.ini
+../../../Templates/TP_Collaborative_V/Media/TP_Collaborative.png
+../../../Templates/TP_Collaborative_V/Media/TP_Collaborative_V_Preview.png
diff --git a/Templates/TP_Collaborative_V/manifest.json b/Templates/TP_Collaborative_V/manifest.json
new file mode 100644
--- /dev/null
+++ b/Templates/TP_Collaborative_V/manifest.json
@@ -0,0 +1,31 @@
+{
+	"Version": 1,
+	"Name":
+	[
+		{
+			"Language": "en",
+			"Text": "Collab Viewer"
+		}
+	],
+	"Description":
+	[
+		{
+			"Language": "en",
+			"Text": "This project template provides navigation and interaction for various platforms (Desktop, VR). Use as a starting point for collaborative Industrial or Architectural projects."
+		}
+	],
+	"AssetTypes": [],
+	"SearchTags": [],
+	"ClassTypes": "",
+	"Category": "",
+	"FocusAsset": "",
+	"Thumbnail": "TP_Collaborative.png",
+	"Screenshots":
+	[
+		"TP_Collaborative_V_Preview.png"
+	],
+	"SortKey": "TP_Collaborative_V",
+	"AdditionalFeaturePacks":
+	[
+	]
+}
